Blueprints: Moved snap node prototype to experimental

diff --git a/Engine/Plugins/Experimental/BlueprintSnapNodes/BlueprintSnapNodes.uplugin b/Engine/Plugins/Experimental/BlueprintSnapNodes/BlueprintSnapNodes.uplugin
new file mode 100644
--- /dev/null
+++ b/Engine/Plugins/Experimental/BlueprintSnapNodes/BlueprintSnapNodes.uplugin
@@ -0,0 +1,25 @@
+{
+	"FileVersion": 3,
+	"Version": 1,
+	"VersionName": "1.0",
+	"FriendlyName": "Blueprint Snap Nodes Prototype",
+	"Description": "Prototype of an alternative way to lay out Blueprint nodes in a more compact manner",
+	"Category": "Experimental",
+	"CreatedBy": "Epic Games, Inc.",
+	"CreatedByURL": "http://epicgames.com",
+	"DocsURL": "",
+	"MarketplaceURL": "",
+	"SupportURL": "",
+	"CanContainContent": false,
+	"IsBetaVersion": true,
+	"Installed": false,
+	"EnabledByDefault": false,
+	"Modules": [
+		{
+			"Name": "BlueprintSnapNodes",
+			"Type": "UncookedOnly",
+			"LoadingPhase": "Default"
+		}
+	],
+	"IsExperimentalVersion": false
+}
\ No newline at end of file
diff --git a/Engine/Plugins/Experimental/BlueprintSnapNodes/Source/BlueprintSnapNodes/BlueprintSnapNodes.Build.cs b/Engine/Plugins/Experimental/BlueprintSnapNodes/Source/BlueprintSnapNodes/BlueprintSnapNodes.Build.cs
new file mode 100644
--- /dev/null
+++ b/Engine/Plugins/Experimental/BlueprintSnapNodes/Source/BlueprintSnapNodes/BlueprintSnapNodes.Build.cs
@@ -0,0 +1,60 @@
+// Copyright Epic Games, Inc. All Rights Reserved.
+
+using UnrealBuildTool;
+
+public class BlueprintSnapNodes : ModuleRules
+{
+	public BlueprintSnapNodes(ReadOnlyTargetRules Target) : base(Target)
+	{
+		PCHUsage = ModuleRules.PCHUsageMode.UseExplicitOrSharedPCHs;
+		
+		PublicIncludePaths.AddRange(
+			new string[] {
+				// ... add public include paths required here ...
+			}
+			);
+				
+		
+		PrivateIncludePaths.AddRange(
+			new string[] {
+				"BlueprintSnapNodes/Private",
+				// ... add other private include paths required here ...
+			}
+			);
+			
+		
+		PublicDependencyModuleNames.AddRange(
+			new string[]
+			{
+				"Core",
+				// ... add other public dependencies that you statically link with here ...
+			}
+			);
+			
+		
+		PrivateDependencyModuleNames.AddRange(
+			new string[]
+			{
+				"CoreUObject",
+				"Engine",
+				"Slate",
+				"SlateCore",
+				"InputCore",
+				"EditorStyle",
+				"GraphEditor",
+				"BlueprintGraph",
+				"EditorFramework",
+				"UnrealEd"
+				// ... add private dependencies that you statically link with here ...	
+			}
+			);
+		
+		
+		DynamicallyLoadedModuleNames.AddRange(
+			new string[]
+			{
+				// ... add any modules that your module loads dynamically here ...
+			}
+			);
+	}
+}
diff --git a/Engine/Plugins/Experimental/BlueprintSnapNodes/Source/BlueprintSnapNodes/Private/BlueprintSnapNodes.cpp b/Engine/Plugins/Experimental/BlueprintSnapNodes/Source/BlueprintSnapNodes/Private/BlueprintSnapNodes.cpp
new file mode 100644
--- /dev/null
+++ b/Engine/Plugins/Experimental/BlueprintSnapNodes/Source/BlueprintSnapNodes/Private/BlueprintSnapNodes.cpp
@@ -0,0 +1,46 @@
+// Copyright Epic Games, Inc. All Rights Reserved.
+
+#include "BlueprintSnapNodes.h"
+#include "EdGraphUtilities.h"
+#include "K2Node_SnapContainer.h"
+#include "SGraphNodeSnapContainer.h"
+
+#define LOCTEXT_NAMESPACE "FBlueprintSnapNodesModule"
+
+class FBlueprintSnapNodeWidgetFactory : public FGraphPanelNodeFactory
+{
+	virtual TSharedPtr<SGraphNode> CreateNode(class UEdGraphNode* InNode) const override
+	{
+		if (UK2Node_SnapContainer* SnapContainer = Cast<UK2Node_SnapContainer>(InNode))
+		{
+			return SNew(SGraphNodeSnapContainer, SnapContainer);
+		}
+		return nullptr;
+	}
+};
+
+
+
+
+
+void FBlueprintSnapNodesModule::StartupModule()
+{
+// 	TSharedPtr<FGameplayTagsGraphPanelPinFactory> GameplayTagsGraphPanelPinFactory = MakeShareable(new FGameplayTagsGraphPanelPinFactory());
+// 	FEdGraphUtilities::RegisterVisualPinFactory(GameplayTagsGraphPanelPinFactory);
+
+	NodeWidgetFactory = MakeShareable(new FBlueprintSnapNodeWidgetFactory());
+	FEdGraphUtilities::RegisterVisualNodeFactory(NodeWidgetFactory);
+}
+
+void FBlueprintSnapNodesModule::ShutdownModule()
+{
+	if (UObjectInitialized())
+	{
+		FEdGraphUtilities::UnregisterVisualNodeFactory(NodeWidgetFactory);
+	}
+	NodeWidgetFactory.Reset();
+}
+
+#undef LOCTEXT_NAMESPACE
+	
+IMPLEMENT_MODULE(FBlueprintSnapNodesModule, BlueprintSnapNodes)
\ No newline at end of file
diff --git a/Engine/Plugins/Experimental/BlueprintSnapNodes/Source/BlueprintSnapNodes/Private/K2Node_SnapContainer.cpp b/Engine/Plugins/Experimental/BlueprintSnapNodes/Source/BlueprintSnapNodes/Private/K2Node_SnapContainer.cpp
new file mode 100644
--- /dev/null
+++ b/Engine/Plugins/Experimental/BlueprintSnapNodes/Source/BlueprintSnapNodes/Private/K2Node_SnapContainer.cpp
@@ -0,0 +1,360 @@
+// Copyright Epic Games, Inc. All Rights Reserved.
+
+#include "K2Node_SnapContainer.h"
+#include "UObject/UnrealType.h"
+#include "UObject/UObjectHash.h"
+#include "UObject/UObjectIterator.h"
+#include "Engine/MemberReference.h"
+#include "Kismet/BlueprintFunctionLibrary.h"
+#include "EdGraphSchema_K2.h"
+#include "EdGraphSchema_K2_Actions.h"
+#include "K2Node_CallFunction.h"
+#include "K2Node_MacroInstance.h"
+#include "K2Node_VariableGet.h"
+#include "Kismet2/BlueprintEditorUtils.h"
+#include "Kismet2/Kismet2NameValidators.h"
+#include "EdGraphUtilities.h"
+//#include "BasicTokenParser.h"
+#include "BlueprintActionDatabaseRegistrar.h"
+#include "DiffResults.h"
+//#include "MathExpressionHandler.h"
+#include "BlueprintNodeSpawner.h"
+
+
+
+#include "K2Node_ExecutionSequence.h"
+
+#define LOCTEXT_NAMESPACE "BlueprintSnapNodes"
+
+template <typename NodeClass>
+NodeClass* SpawnNewNode(UEdGraph* ParentGraph, const FVector2D Location)
+{
+	NodeClass* NewNode = NewObject<NodeClass>(ParentGraph, NodeClass::StaticClass());
+	check(NewNode != nullptr);
+	NewNode->CreateNewGuid();
+
+	// position the node before invoking PostSpawnDelegate (in case it wishes to modify this positioning)
+	NewNode->NodePosX = Location.X;
+	NewNode->NodePosY = Location.Y;
+
+	NewNode->SetFlags(RF_Transactional);
+	NewNode->AllocateDefaultPins();
+	NewNode->PostPlacedNewNode();
+
+	ParentGraph->Modify();
+
+	ParentGraph->AddNode(NewNode, /*bFromUI =*/true, /*bSelectNewNode =*/false);
+
+	return NewNode;
+}
+
+
+/*******************************************************************************
+ * Static Helpers
+*******************************************************************************/
+
+/**
+ * Helper function for deleting all the nodes from a specified graph. Does not
+ * delete any tunnel in/out nodes (to preserve the tunnel).
+ * 
+ * @param  Graph	The graph you want to delete nodes in.
+ */
+static void DeleteGeneratedNodesInGraph(UEdGraph* Graph)
+{
+	UBlueprint* Blueprint = FBlueprintEditorUtils::FindBlueprintForGraphChecked(Graph);
+	for (int32 NodeIndex = 0; NodeIndex < Graph->Nodes.Num(); )
+	{
+		UEdGraphNode* Node = Graph->Nodes[NodeIndex];
+		if (ExactCast<UK2Node_Tunnel>(Node) != NULL)
+		{
+			++NodeIndex;
+		}
+		else
+		{
+			FBlueprintEditorUtils::RemoveNode(Blueprint, Node, true);
+		}
+	}
+}
+/**
+ * Sets or clears the error on a specific node. If the ErrorText is empty, then
+ * it resets the error state on the node. If actual error text is supplied, 
+ * then the node is flagged as having an error, and the string is appended to 
+ * the node's error message.
+ * 
+ * @param  Node			The node you wish to modify.
+ * @param  ErrorText	The text you want to append to the node's error message (empty if you want to clear any errors).
+ */
+static void SetNodeError(UEdGraphNode* Node, FText const& ErrorText)
+{
+	if (ErrorText.IsEmpty())
+	{
+		Node->ErrorMsg.Empty();
+		Node->ErrorType = EMessageSeverity::Info;
+		Node->bHasCompilerMessage = false;
+	}
+	else if (Node->bHasCompilerMessage)
+	{
+		Node->ErrorMsg += TEXT("\n") + ErrorText.ToString();
+		Node->ErrorType = EMessageSeverity::Error;
+	}
+	else
+	{
+		Node->ErrorMsg = ErrorText.ToString();
+		Node->ErrorType = EMessageSeverity::Error;
+		Node->bHasCompilerMessage = true;
+	}
+}
+
+
+/*******************************************************************************
+ * FLayoutVisitor
+*******************************************************************************/
+
+/**
+ * This class is utilized to help layout math expression nodes by traversing the
+ * expression tree and cataloging each expression node's depth. From the tree's 
+ * depth we can determine the width of the the graph (an where to place each K2 node):
+ *
+ *    _
+ *   |            [_]---[_]
+ *   |           /
+ * height   [_]--       [_]--[_]---[_]
+ *   |           \     /
+ *   |_           [_]---[_]
+ *
+ *		    ^-------depth/width-------^
+ */
+#if 0
+class FLayoutVisitor : public FExpressionVisitor
+{
+public:
+	/** Tracks the horizontal (depth) placement for each expression node encountered */
+	TMap<IFExpressionNode*, int32> DepthChart;
+	/** Tracks the vertical (height) placement for each expression node encountered */
+	TMap<IFExpressionNode*, int32> HeightChart;
+	/** Tracks the total height (value) at each depth (key) */
+	TMap<int32, int32> DepthHeightLookup;
+
+	/** */
+	FLayoutVisitor()
+		: CurrentDepth(0)
+		, MaximumDepth(0)
+	{
+	}
+
+	/**
+	 * Retrieves the total depth (or graph width) of the previously traversed
+	 * expression tree.
+	 */
+	int32 GetMaximumDepth() const
+	{
+		return MaximumDepth;
+	}
+
+	/**
+	 * Resets this tree visitor so that it can accurately parse another 
+	 * expression tree (else the results would stack up).
+	 */
+	void Clear() 
+	{
+		CurrentDepth = 0;
+		MaximumDepth = 0;
+		DepthChart.Empty();
+		HeightChart.Empty();
+		DepthHeightLookup.Empty();
+	}
+
+private:
+	/**
+	 * From the FExpressionVisitor interface, a generic choke point for visiting 
+	 * all expression nodes.
+	 *
+	 * @return True to continue traversing the tree, false to abort. 
+	 */
+	virtual bool VisitUnhandled(class IFExpressionNode& Node, EVisitPhase Phase) override
+	{
+		if (Phase == FExpressionVisitor::VISIT_Pre)
+		{
+			++CurrentDepth;
+			MaximumDepth = FMath::Max(CurrentDepth, MaximumDepth);
+		}
+		else
+		{
+			if (Phase == FExpressionVisitor::VISIT_Post)
+			{
+				--CurrentDepth;
+			}
+			// else leaf
+
+			// CurrentHeight represents how many nodes have already been placed 
+			// at this depth
+			int32& CurrentHeight = DepthHeightLookup.FindOrAdd(CurrentDepth);
+
+			DepthChart.Add(&Node, CurrentDepth);
+			HeightChart.Add(&Node, CurrentHeight);
+
+			// since we just placed another node at this depth, increase the 
+			// height count
+			++CurrentHeight;
+		}
+
+		// let the tree traversal continue! don't abort it!
+		return true;
+	}
+
+private:
+	int32 CurrentDepth;
+	int32 MaximumDepth;
+};
+#endif
+
+
+/*******************************************************************************
+ * UK2Node_SnapContainer
+*******************************************************************************/
+
+UK2Node_SnapContainer::UK2Node_SnapContainer(const FObjectInitializer& ObjectInitializer)
+	: Super(ObjectInitializer)
+{
+}
+
+void UK2Node_SnapContainer::PostEditChangeProperty(struct FPropertyChangedEvent& PropertyChangedEvent)
+{
+	Super::PostEditChangeProperty(PropertyChangedEvent);
+}
+
+void UK2Node_SnapContainer::GetMenuActions(FBlueprintActionDatabaseRegistrar& ActionRegistrar) const
+{
+	UClass* ActionKey = GetClass();
+	if (ActionRegistrar.IsOpenForRegistration(ActionKey))
+	{
+		UBlueprintNodeSpawner* NodeSpawner = UBlueprintNodeSpawner::Create(GetClass());
+		check(NodeSpawner != nullptr);
+
+		ActionRegistrar.AddBlueprintAction(ActionKey, NodeSpawner);
+	}
+}
+
+TSharedPtr<class INameValidatorInterface> UK2Node_SnapContainer::MakeNameValidator() const
+{
+	// we'll let our parser mark the node for errors after the face (once the 
+	// name is submitted)... parsing it with every character could be slow
+	return MakeShareable(new FDummyNameValidator(EValidatorResult::Ok));
+}
+
+FText UK2Node_SnapContainer::GetNodeTitle(ENodeTitleType::Type TitleType) const
+{
+	if (TitleType == ENodeTitleType::MenuTitle)
+	{
+		return LOCTEXT("AddSnapContainer", "Add Snap Container...");
+	}
+	
+	return LOCTEXT("PlacedSnapContainerNodeTitle", "Snap Container");
+}
+
+void UK2Node_SnapContainer::ReconstructNode()
+{
+	if (!HasAnyFlags(RF_NeedLoad))
+	{
+		const UEdGraphSchema_K2* K2Schema = GetDefault<UEdGraphSchema_K2>();
+
+		if (UK2Node_ExecutionSequence* Sequence = Cast<UK2Node_ExecutionSequence>(RootNode))
+		{
+			// Compact the sequence node, and make sure there is an empty entry at the end of the node (always keep an open drag target)
+			//@TODO: Do that thing I said
+			UEdGraphPin* LastExecPin = nullptr;
+			for (UEdGraphPin* Pin : Sequence->Pins)
+			{
+				if (K2Schema->IsExecPin(*Pin) && (Pin->Direction == EGPD_Output))
+				{
+					LastExecPin = Pin;
+				}
+			}
+
+			if ((LastExecPin == nullptr) || (LastExecPin->LinkedTo.Num() > 0))
+			{
+				Sequence->AddInputPin();
+			}
+		}
+	}
+	Super::ReconstructNode();
+}
+
+void UK2Node_SnapContainer::AllocateDefaultPins()
+{
+	Super::AllocateDefaultPins();
+
+// 	// AllocateDefaultPins() gets called too early during initial node creation, before BoundGraph was created, so
+// 	// the first time thru we need to add a pin manually.  After that it will be added automatically because the
+// 	// inner tunnel has it
+// 	if (FindPin(UEdGraphSchema_K2::PN_Execute, EGPD_Input) == nullptr)
+// 	{
+// 		const UEdGraphSchema_K2* K2Schema = GetDefault<UEdGraphSchema_K2>();
+// 		CreatePin(EGPD_Input, UEdGraphSchema_K2::PC_Exec, FString(), nullptr, UEdGraphSchema_K2::PN_Execute);
+// 	}
+}
+
+void UK2Node_SnapContainer::PostPlacedNewNode()
+{
+	Super::PostPlacedNewNode();
+
+	const UEdGraphSchema_K2* K2Schema = GetDefault<UEdGraphSchema_K2>();
+
+	// Create a root node (hardcoded to a sequence for now)
+	UK2Node_ExecutionSequence* TopLevelSequence = SpawnNewNode<UK2Node_ExecutionSequence>(BoundGraph, FVector2D::ZeroVector);
+	RootNode = TopLevelSequence;
+
+	// Create an exec pin on the entry tunnel (gets mirrored onto the outer node automatically)
+	UK2Node_Tunnel* EntryNode = GetEntryNode();
+	FEdGraphPinType ExecPinType;
+	ExecPinType.PinCategory = UEdGraphSchema_K2::PC_Exec;
+	UEdGraphPin* TunnelExecIn = EntryNode->CreateUserDefinedPin(UEdGraphSchema_K2::PN_Execute, ExecPinType, EGPD_Input, /*bUseUniqueName=*/ false);
+
+	// Make sure the inner tunnel is connected to the sequence node
+	EntryNode->NodePosX = -200;
+	UEdGraphPin* SequenceExec = TopLevelSequence->FindPinChecked(UEdGraphSchema_K2::PN_Execute);
+	K2Schema->TryCreateConnection(TunnelExecIn, SequenceExec);
+
+	// Position the exit node
+	UK2Node_Tunnel* ExitNode = GetExitNode();
+	ExitNode->NodePosX = 1024;
+
+	// Make sure we have the top level exec pin to match the one we put into the entry node
+	// (AllocateDefaultPins() gets called too early during initial node creation, before BoundGraph was created)
+//	AllocateDefaultPins();
+}
+
+FText UK2Node_SnapContainer::GetKeywords() const
+{
+	return LOCTEXT("SnapContainerKeywords", "Block Node");
+}
+
+void UK2Node_SnapContainer::FindDiffs(class UEdGraphNode* OtherNode, struct FDiffResults& Results )
+{
+	UK2Node_SnapContainer* MathExpression1 = this;
+	UK2Node_SnapContainer* MathExpression2 = Cast<UK2Node_SnapContainer>(OtherNode);
+
+	// Compare the visual display of a math expression (the visual display involves consolidating variable Guid's into readable parameters)
+	FText Expression1 = MathExpression1->GetNodeTitle(ENodeTitleType::EditableTitle);
+	FText Expression2 = MathExpression2->GetNodeTitle(ENodeTitleType::EditableTitle);
+	if (Expression1.CompareTo(Expression2) != 0)
+	{
+		FDiffSingleResult Diff;
+		Diff.Node1 = MathExpression2;
+		Diff.Node2 = MathExpression1;
+
+		Diff.Diff = EDiffType::NODE_PROPERTY;
+		FText NodeName = GetNodeTitle(ENodeTitleType::ListView);
+
+		FFormatNamedArguments Args;
+		Args.Add(TEXT("Expression1"), Expression1);
+		Args.Add(TEXT("Expression2"), Expression2);
+
+		Diff.ToolTip =  FText::Format(LOCTEXT("DIF_MathExpressionToolTip", "Math Expression '{Expression1}' changed to '{Expression2}'"), Args);
+		Diff.DisplayColor = FLinearColor(0.85f,0.71f,0.25f);
+		Diff.DisplayString = FText::Format(LOCTEXT("DIF_MathExpression", "Math Expression '{Expression1}' changed to '{Expression2}'"), Args);
+		Results.Add(Diff);
+	}
+}
+
+#undef LOCTEXT_NAMESPACE
diff --git a/Engine/Plugins/Experimental/BlueprintSnapNodes/Source/BlueprintSnapNodes/Private/K2Node_SnapContainer.h b/Engine/Plugins/Experimental/BlueprintSnapNodes/Source/BlueprintSnapNodes/Private/K2Node_SnapContainer.h
new file mode 100644
--- /dev/null
+++ b/Engine/Plugins/Experimental/BlueprintSnapNodes/Source/BlueprintSnapNodes/Private/K2Node_SnapContainer.h
@@ -0,0 +1,57 @@
+// Copyright Epic Games, Inc. All Rights Reserved.
+
+#pragma once
+
+#include "CoreMinimal.h"
+#include "UObject/ObjectMacros.h"
+#include "Kismet2/CompilerResultsLog.h"
+#include "EdGraph/EdGraphNodeUtils.h"
+#include "K2Node_Composite.h"
+#include "K2Node_SnapContainer.generated.h"
+
+class FBlueprintActionDatabaseRegistrar;
+class INameValidatorInterface;
+
+//@TODO: Write a comment
+UCLASS()
+class UK2Node_SnapContainer : public UK2Node_Composite
+{
+	GENERATED_BODY()
+
+public:
+	UPROPERTY()
+	UK2Node* RootNode;
+
+public:
+	UK2Node_SnapContainer(const FObjectInitializer& ObjectInitializer);
+
+	//~ Begin UObject Interface
+	virtual void PostEditChangeProperty(struct FPropertyChangedEvent& PropertyChangedEvent) override;
+	//~ End UObject Interface
+
+	//~ Begin UEdGraphNode Interface
+	virtual TSharedPtr<class INameValidatorInterface> MakeNameValidator() const override;
+	virtual FText GetNodeTitle(ENodeTitleType::Type TitleType) const override;
+	virtual void ReconstructNode() override;
+	virtual void FindDiffs(class UEdGraphNode* OtherNode, struct FDiffResults& Results) override;
+	virtual bool ShouldMergeChildGraphs() const override { return true; }
+	virtual void AllocateDefaultPins() override;
+	virtual void PostPlacedNewNode() override;
+	virtual FText GetKeywords() const override;
+	//~ End UEdGraphNode Interface
+
+	//~ Begin UK2Node Interface
+	virtual void GetMenuActions(FBlueprintActionDatabaseRegistrar& ActionRegistrar) const override;
+	virtual bool IsNodePure() const { return false; }
+	//~ End UK2Node Interface
+
+	//~ Begin UK2Node_EditablePinBase Interface
+	virtual bool CanCreateUserDefinedPin(const FEdGraphPinType& InPinType, EEdGraphPinDirection InDesiredDirection, FText& OutErrorMessage) override { return false; }
+	//~ End UK2Node_EditablePinBase Interface
+
+private:
+	/** Cached so we don't have to regenerate it when the graph is recompiled */
+	TSharedPtr<class FCompilerResultsLog> CachedMessageLog;
+};
+
+
diff --git a/Engine/Plugins/Experimental/BlueprintSnapNodes/Source/BlueprintSnapNodes/Private/SGraphNodeSnapContainer.cpp b/Engine/Plugins/Experimental/BlueprintSnapNodes/Source/BlueprintSnapNodes/Private/SGraphNodeSnapContainer.cpp
new file mode 100644
--- /dev/null
+++ b/Engine/Plugins/Experimental/BlueprintSnapNodes/Source/BlueprintSnapNodes/Private/SGraphNodeSnapContainer.cpp
@@ -0,0 +1,295 @@
+// Copyright Epic Games, Inc. All Rights Reserved.
+
+#include "SGraphNodeSnapContainer.h"
+#include "EdGraph/EdGraph.h"
+#include "Widgets/SBoxPanel.h"
+#include "Framework/Application/SlateApplication.h"
+#include "Widgets/Layout/SSpacer.h"
+#include "Widgets/Images/SImage.h"
+#include "Widgets/SToolTip.h"
+#include "Styling/CoreStyle.h"
+#include "GraphEditorSettings.h"
+#include "SCommentBubble.h"
+#include "K2Node_Composite.h"
+#include "SGraphPreviewer.h"
+#include "IDocumentationPage.h"
+#include "IDocumentation.h"
+#include "Widgets/Text/SInlineEditableTextBlock.h"
+
+#include "SGraphSnapContainerRow.h"
+#include "K2Node_SnapContainer.h"
+
+#define LOCTEXT_NAMESPACE "BlueprintSnapNodes"
+
+/////////////////////////////////////////////////////
+// SGraphNodeSnapContainer
+
+void SGraphNodeSnapContainer::Construct(const FArguments& InArgs, UK2Node_SnapContainer* InNode)
+{
+	GraphNode = InNode;
+
+	SetCursor(EMouseCursor::CardinalCross);
+
+	UpdateGraphNode();
+}
+
+void SGraphNodeSnapContainer::UpdateGraphNode()
+{
+	InputPins.Empty();
+	OutputPins.Empty();
+	
+	// Reset variables that are going to be exposed, in case we are refreshing an already setup node.
+	RightNodeBox.Reset();
+	LeftNodeBox.Reset();
+
+	SetupErrorReporting();
+	TSharedPtr<SNodeTitle> NodeTitle = SNew(SNodeTitle, GraphNode);
+
+
+	UK2Node_SnapContainer* SnapContainer = CastChecked<UK2Node_SnapContainer>(GraphNode);
+	TSharedRef<SWidget> SnapRegionWidget = FGraphSnapContainerBuilder::CreateSnapContainerWidgets(SnapContainer->BoundGraph, SnapContainer->RootNode);
+
+	//
+	//             ______________________
+	//            |      TITLE AREA      |
+	//            +-------+------+-------+
+	//            | (>) L |      | R (>) |
+	//            | (>) E |      | I (>) |
+	//            | (>) F |      | G (>) |
+	//            | (>) T |      | H (>) |
+	//            |       |      | T (>) |
+	//            |_______|______|_______|
+	//
+	this->ContentScale.Bind( this, &SGraphNode::GetContentScale );
+	this->GetOrAddSlot( ENodeZone::Center )
+		.HAlign(HAlign_Center)
+		.VAlign(VAlign_Center)
+		[
+			SNew(SBorder)
+			.BorderImage( FEditorStyle::GetBrush( "Graph.CollapsedNode.Body" ) )
+			.Padding(0)
+			[
+				SNew(SOverlay)
+				+SOverlay::Slot()
+				[
+					SNew(SImage)
+					.Image( FEditorStyle::GetBrush("Graph.CollapsedNode.BodyColorSpill") )
+					.ColorAndOpacity( this, &SGraphNode::GetNodeTitleColor )
+				]
+				+SOverlay::Slot()
+				[
+					SNew(SVerticalBox)
+					+SVerticalBox::Slot()
+					.AutoHeight()
+					.HAlign(HAlign_Fill)
+					.VAlign(VAlign_Top)
+					[
+						SNew(SOverlay)
+						+SOverlay::Slot()
+						.HAlign(HAlign_Left)
+						.VAlign(VAlign_Center)
+						[
+							SNew(SBorder)
+							.BorderImage( FEditorStyle::GetBrush("NoBorder") )  // Graph.CollapsedNode.ColorSpill
+							.Padding( FMargin(10,5,30,3) )
+							[
+								SNew(SVerticalBox)
+								+SVerticalBox::Slot()
+								.AutoHeight()
+								.HAlign(HAlign_Fill)
+								.VAlign(VAlign_Top)
+								[
+									SNew(SVerticalBox)
+									+SVerticalBox::Slot()
+										.AutoHeight()
+									[
+										SAssignNew(InlineEditableText, SInlineEditableTextBlock)
+										.Style( FEditorStyle::Get(), "Graph.Node.NodeTitleInlineEditableText" )
+										.Text( NodeTitle.Get(), &SNodeTitle::GetHeadTitle )
+										.OnVerifyTextChanged(this, &SGraphNodeSnapContainer::OnVerifyNameTextChanged)
+										.OnTextCommitted(this, &SGraphNodeSnapContainer::OnNameTextCommited)
+										.IsReadOnly( this, &SGraphNodeSnapContainer::IsNameReadOnly )
+										.IsSelected(this, &SGraphNodeSnapContainer::IsSelectedExclusively)
+									]
+									+SVerticalBox::Slot()
+										.AutoHeight()
+									[
+										NodeTitle.ToSharedRef()
+									]
+								]
+								+SVerticalBox::Slot()
+								.AutoHeight()
+								.Padding(1.0f)
+								[
+									ErrorReporting->AsWidget()
+								]
+							]
+						]
+					]
+					+SVerticalBox::Slot()
+					.AutoHeight()
+					.HAlign(HAlign_Fill)
+					.VAlign(VAlign_Top)
+					[
+						CreateNodeBody()
+					]
+					+SVerticalBox::Slot()
+					.AutoHeight()
+					.HAlign(HAlign_Fill)
+					.VAlign(VAlign_Top)
+					[
+						SnapRegionWidget
+					]
+				]
+			]
+		];
+	// Create comment bubble
+	TSharedPtr<SCommentBubble> CommentBubble;
+	const FSlateColor CommentColor = GetDefault<UGraphEditorSettings>()->DefaultCommentNodeTitleColor;
+
+	SAssignNew( CommentBubble, SCommentBubble )
+	.GraphNode( GraphNode )
+	.Text( this, &SGraphNode::GetNodeComment )
+	.OnTextCommitted( this, &SGraphNode::OnCommentTextCommitted )
+	.ColorAndOpacity( CommentColor )
+	.AllowPinning( true )
+	.EnableTitleBarBubble( true )
+	.EnableBubbleCtrls( true )
+	.GraphLOD( this, &SGraphNode::GetCurrentLOD )
+	.IsGraphNodeHovered( this, &SGraphNode::IsHovered );
+
+	GetOrAddSlot( ENodeZone::TopCenter )
+	.SlotOffset( TAttribute<FVector2D>( CommentBubble.Get(), &SCommentBubble::GetOffset ))
+	.SlotSize( TAttribute<FVector2D>( CommentBubble.Get(), &SCommentBubble::GetSize ))
+	.AllowScaling( TAttribute<bool>( CommentBubble.Get(), &SCommentBubble::IsScalingAllowed ))
+	.VAlign( VAlign_Top )
+	[
+		CommentBubble.ToSharedRef()
+	];
+
+	CreatePinWidgets();
+}
+
+UEdGraph* SGraphNodeSnapContainer::GetInnerGraph() const
+{
+	UK2Node_Composite* CompositeNode = CastChecked<UK2Node_Composite>(GraphNode);
+	return CompositeNode->BoundGraph;
+}
+
+TSharedPtr<SToolTip> SGraphNodeSnapContainer::GetComplexTooltip()
+{
+	if (UEdGraph* BoundGraph = GetInnerGraph())
+	{
+		struct LocalUtils
+		{
+			static bool IsInteractive()
+			{
+				const FModifierKeysState ModifierKeys = FSlateApplication::Get().GetModifierKeys();
+				return ( ModifierKeys.IsAltDown() && ModifierKeys.IsControlDown() );
+			}
+		};
+
+		TSharedPtr<SToolTip> FinalToolTip = NULL;
+		TSharedPtr<SVerticalBox> Container = NULL;
+		SAssignNew(FinalToolTip, SToolTip)
+		.IsInteractive_Static(&LocalUtils::IsInteractive)
+		[
+			SAssignNew(Container, SVerticalBox)
+			+SVerticalBox::Slot()
+			.AutoHeight()
+			[
+				SNew( STextBlock )
+				.Text(this, &SGraphNodeSnapContainer::GetTooltipTextForNode)
+				.Font(FCoreStyle::GetDefaultFontStyle("Regular", 8))
+				.WrapTextAt(160.0f)
+			]
+			+SVerticalBox::Slot()
+			.AutoHeight()
+			[
+				// Create preview for the tooltip, make sure to disable state overlays to prevent
+				// PIE and read-only borders obscuring the graph
+				SNew(SGraphPreviewer, BoundGraph)
+				.CornerOverlayText(this, &SGraphNodeSnapContainer::GetPreviewCornerText)
+				.ShowGraphStateOverlay(false)
+			]
+		];
+
+		// Check to see whether this node has a documentation excerpt. If it does, create a doc box for the tooltip
+		TSharedRef<IDocumentationPage> DocPage = IDocumentation::Get()->GetPage(GraphNode->GetDocumentationLink(), NULL);
+		if(DocPage->HasExcerpt(GraphNode->GetDocumentationExcerptName()))
+		{
+			Container->AddSlot()
+			.AutoHeight()
+			.Padding(FMargin( 0.0f, 5.0f ))
+			[
+				IDocumentation::Get()->CreateToolTip(FText::FromString("Documentation"), NULL, GraphNode->GetDocumentationLink(), GraphNode->GetDocumentationExcerptName())
+			];
+		}
+
+		return FinalToolTip;
+	}
+	else
+	{
+		return SNew(SToolTip)
+			[
+				SNew(SVerticalBox)
+				+SVerticalBox::Slot()
+				.AutoHeight()
+				[
+					SNew( STextBlock )
+					.Text(NSLOCTEXT("CompositeNode", "CompositeNodeInvalidGraphMessage", "ERROR: Invalid Graph"))
+					.Font(FCoreStyle::GetDefaultFontStyle("Regular", 8))
+					.WrapTextAt(160.0f)
+				]
+			];
+	}
+
+}
+
+FText SGraphNodeSnapContainer::GetPreviewCornerText() const
+{
+	return LOCTEXT("SnapGraphCornerText", "Snap Magic");
+}
+
+FText SGraphNodeSnapContainer::GetTooltipTextForNode() const
+{
+	return GraphNode->GetTooltipText();
+}
+
+TSharedRef<SWidget> SGraphNodeSnapContainer::CreateNodeBody()
+{
+	if( GraphNode && GraphNode->Pins.Num() > 0 )
+	{
+		// Create the input and output pin areas if there are pins
+		return SNew(SBorder)
+			.BorderImage( FEditorStyle::GetBrush("NoBorder") )
+			.HAlign(HAlign_Fill)
+			.VAlign(VAlign_Fill)
+			.Padding( FMargin(0,3) )
+			[
+				SNew(SHorizontalBox)
+				+SHorizontalBox::Slot()
+				.HAlign(HAlign_Left)
+				.FillWidth(1.0f)
+				[
+					// LEFT
+					SAssignNew(LeftNodeBox, SVerticalBox)
+				]
+				+SHorizontalBox::Slot()
+				.AutoWidth()
+				.HAlign(HAlign_Right)
+				[
+					// RIGHT
+					SAssignNew(RightNodeBox, SVerticalBox)
+				]
+			];
+	}
+	else
+	{
+		// Create a spacer so the node has some body to it
+		return SNew(SSpacer)
+			.Size(FVector2D(100.f, 50.f));
+	}
+}
+
+#undef LOCTEXT_NAMESPACE
diff --git a/Engine/Plugins/Experimental/BlueprintSnapNodes/Source/BlueprintSnapNodes/Private/SGraphNodeSnapContainer.h b/Engine/Plugins/Experimental/BlueprintSnapNodes/Source/BlueprintSnapNodes/Private/SGraphNodeSnapContainer.h
new file mode 100644
--- /dev/null
+++ b/Engine/Plugins/Experimental/BlueprintSnapNodes/Source/BlueprintSnapNodes/Private/SGraphNodeSnapContainer.h
@@ -0,0 +1,35 @@
+// Copyright Epic Games, Inc. All Rights Reserved.
+
+#pragma once
+
+#include "CoreMinimal.h"
+#include "Widgets/DeclarativeSyntaxSupport.h"
+#include "Widgets/SWidget.h"
+#include "KismetNodes/SGraphNodeK2Base.h"
+
+class SToolTip;
+class UEdGraph;
+class UK2Node_SnapContainer;
+
+class SGraphNodeSnapContainer : public SGraphNodeK2Base
+{
+public:
+	SLATE_BEGIN_ARGS(SGraphNodeSnapContainer){}
+	SLATE_END_ARGS()
+	typedef SGraphNodeK2Base Super;
+
+	void Construct(const FArguments& InArgs, UK2Node_SnapContainer* InNode);
+
+	// SGraphNode interface
+	virtual void UpdateGraphNode() override;
+	virtual TSharedPtr<SToolTip> GetComplexTooltip() override;
+	// End of SGraphNode interface
+protected:
+	virtual UEdGraph* GetInnerGraph() const;
+
+private:
+	FText GetPreviewCornerText() const;
+	FText GetTooltipTextForNode() const;
+
+	TSharedRef<SWidget> CreateNodeBody();
+};
diff --git a/Engine/Plugins/Experimental/BlueprintSnapNodes/Source/BlueprintSnapNodes/Private/SGraphSnapContainerRow.cpp b/Engine/Plugins/Experimental/BlueprintSnapNodes/Source/BlueprintSnapNodes/Private/SGraphSnapContainerRow.cpp
new file mode 100644
--- /dev/null
+++ b/Engine/Plugins/Experimental/BlueprintSnapNodes/Source/BlueprintSnapNodes/Private/SGraphSnapContainerRow.cpp
@@ -0,0 +1,627 @@
+// Copyright Epic Games, Inc. All Rights Reserved.
+
+#include "SGraphSnapContainerRow.h"
+#include "EdGraph/EdGraph.h"
+#include "Widgets/SBoxPanel.h"
+#include "Framework/Application/SlateApplication.h"
+#include "Widgets/Layout/SSpacer.h"
+#include "Widgets/Images/SImage.h"
+#include "Widgets/SToolTip.h"
+#include "GraphEditorSettings.h"
+#include "SCommentBubble.h"
+#include "K2Node_Composite.h"
+#include "SGraphPreviewer.h"
+#include "IDocumentationPage.h"
+#include "IDocumentation.h"
+#include "Widgets/Text/SInlineEditableTextBlock.h"
+
+#include "NodeFactory.h"
+#include "K2Node_ExecutionSequence.h"
+#include "SGraphPin.h"
+#include "GraphEditorDragDropAction.h"
+#include "K2Node.h"
+#include "K2Node_VariableGet.h"
+#include "K2Node_VariableSet.h"
+#include "K2Node_SnapContainer.h"
+#include "EdGraphSchema_K2.h"
+#include "Widgets/Layout/SBox.h"
+
+#define LOCTEXT_NAMESPACE "BlueprintSnapNodes"
+
+//@TODO: Figure out how to handle self nodes in a reasonable way (it currently does "self.", could drag an Other onto it of the same type)
+
+//@TODO: No provision for constructing one of these without drag drop (e.g., a dropdown to pick a method on a self pin)
+//@TODO: Drag-dropping a variable onto an existing variable doesn't work right (not a bug here, same issue happens if you drag a float var onto a float input pin that already has a connection; it drops the get node but fails to replace the existing one)
+//@TODO: How to deal with replaced stuff (after the above note is fixed), e.g., a node in the subgraph that is no longer wired to anything.  Could just clean unconnected nodes on each edit, but seems messy.
+
+//@TODO: Tooltip is wrong on variable get nodes (and they're just ugly fake versions of the real thing)
+//@TODO: Styling in general is ugly
+
+//@TODO: Drag-dropping a connection onto one of the variable 'pins' should probably do something or provide an explicit tooltip why it won't
+// (something might be adding an input pin to the composite that gets wired up internally and adding a generated name for the inner pin that gets wired up)
+// This is a place where a lot of the stuff that math expression nodes do may be useful RE: any old variable 'works' and just becomes an input or output ref
+
+//@TODO: This should probably get collapsed with the code in SPinTypeSelector::GetTypeDescription (this one is eventually going to be better!)
+// It's not perfect tho (just does [] for array/set/map and & for ref
+FText GetTypeDescription(const FEdGraphPinType& PinType)
+{
+	const FName PinSubCategory = PinType.PinSubCategory;
+	const UObject* PinSubCategoryObject = PinType.PinSubCategoryObject.Get();
+
+	FText BasicText;
+	if ((PinSubCategory != UEdGraphSchema_K2::PSC_Bitmask) && (PinSubCategoryObject != nullptr))
+	{
+		if (const UField* Field = Cast<const UField>(PinSubCategoryObject))
+		{
+			BasicText = Field->GetDisplayNameText();
+		}
+		else
+		{
+			BasicText = FText::AsCultureInvariant(PinSubCategoryObject->GetName());
+		}
+	}
+	else
+	{
+		BasicText = UEdGraphSchema_K2::GetCategoryText(PinType.PinCategory, true);
+	}
+
+	switch (PinType.ContainerType)
+	{
+	default:
+		break;
+	case EPinContainerType::Array:
+		BasicText = FText::Format(LOCTEXT("PinTypeIsArray", "TArray<{0}>"), BasicText);
+		break;
+	case EPinContainerType::Set:
+		BasicText = FText::Format(LOCTEXT("PinTypeIsSet", "TSet<{0}>"), BasicText);
+		break;
+	case EPinContainerType::Map:
+		BasicText = FText::Format(LOCTEXT("PinTypeIsMap", "TMap<{0}, TODO>"), BasicText);
+		break;
+	}
+
+	if (PinType.bIsConst)
+	{
+		BasicText = FText::Format(LOCTEXT("PinTypeIsConst", "const {0}"), BasicText);
+	}
+
+	if (PinType.bIsReference)
+	{
+		BasicText = FText::Format(LOCTEXT("PinTypeIsReference", "{0}&"), BasicText);
+	}
+
+	return BasicText;
+}
+
+
+/////////////////////////////////////////////////////
+// FGraphSnapContainerBuilder
+
+TSharedRef<SWidget> FGraphSnapContainerBuilder::CreateSnapContainerWidgets(UEdGraph* ModelGraph, UEdGraphNode* RootNode)
+{
+	FGraphSnapContainerBuilder Builder(ModelGraph);
+
+	const UEdGraphSchema_K2* Schema = GetDefault<UEdGraphSchema_K2>();
+	UEdGraphPin* RootPin = Schema->FindExecutionPin(*RootNode, EGPD_Input);
+
+	return Builder.MakeNodeWidget(RootNode, RootPin);
+}
+
+FGraphSnapContainerBuilder::FGraphSnapContainerBuilder(UEdGraph* InGraph)
+	: Graph(InGraph)
+{
+}
+
+TSharedRef<SWidget> FGraphSnapContainerBuilder::MakeNodeWidget(UEdGraphNode* Node, UEdGraphPin* FromPin)
+{
+	// Make sure we visit each node only once; to prevent infinite recursion
+	bool bAlreadyInSet = false;
+	VisitedNodes.Add(Node, &bAlreadyInSet);
+	if (bAlreadyInSet)
+	{
+		return SNullWidget::NullWidget;
+		// SNew(STextBlock).Text(LOCTEXT("RecursionOccurredInNodeGraphMessage", "RECURSION"));
+	}
+
+	// Evaluate the node to gather information that many of the widgets will want
+	const UEdGraphSchema* Schema = Node->GetSchema();
+	TArray<UEdGraphPin*> InputPins;
+	TArray<UEdGraphPin*> OutputPins;
+	int32 InputPinCount = 0;
+	int32 OutputPinCount = 0;
+	for (auto PinIt = Node->Pins.CreateConstIterator(); PinIt; ++PinIt)
+	{
+		UEdGraphPin* Pin = *PinIt;
+		if (!Pin->bHidden)
+		{
+			if (Pin->Direction == EGPD_Input)
+			{
+				InputPins.Add(Pin);
+				++InputPinCount;
+			}
+			else
+			{
+				OutputPins.Add(Pin);
+				++OutputPinCount;
+			}
+		}
+	}
+
+	// Determine if the node is impure
+	UK2Node* K2Node = Cast<UK2Node>(Node);
+	const bool bIsPure = (K2Node != nullptr) && K2Node->IsNodePure();
+
+	if (UK2Node_ExecutionSequence* Sequence = Cast<UK2Node_ExecutionSequence>(Node))
+	{
+		TSharedRef<SVerticalBox> VerticalPinBox = SNew(SVerticalBox);
+
+		for (UEdGraphPin* OutputExecPin : OutputPins)
+		{
+			VerticalPinBox->AddSlot()
+			.Padding(FMargin(12.0f, 0.0f, 12.0f, 4.0f))
+			.AutoHeight()
+			[
+				SNew(SBox)
+				.Padding(FMargin(8.0f, 3.0f))
+				[
+					MakePinWidget(OutputExecPin)
+				]
+			];
+		}
+
+		return VerticalPinBox;
+	}
+	else if ((OutputPinCount != 1) || !bIsPure)
+	{
+		// It's a function call
+		check(K2Node);
+		return MakeFunctionCallWidget(K2Node);
+
+#if 0
+		// The source node is impure or has multiple outputs, so cannot be directly part of this pure expression
+		// Instead show it as a special sort of variable get
+		FFormatNamedArguments Args;
+		Args.Add(TEXT("NodeTitle"), Node->GetNodeTitle(ENodeTitleType::ListView));
+		Args.Add(TEXT("PinName"), FromPin->GetDisplayName());
+		const FText EffectiveVariableName = FText::Format(LOCTEXT("NodeTitleWithPinNameImpure", "{NodeTitle}_{PinName} impure"), Args );
+
+		return SNew(SOverlay)
+			+ SOverlay::Slot()
+			[
+				SNew(SImage)
+				.Image( FEditorStyle::GetBrush("KismetExpression.ReadAutogeneratedVariable.Body") )
+				.ColorAndOpacity( Schema->GetPinTypeColor(FromPin->PinType) )
+			]
+			+ SOverlay::Slot()
+			.Padding( FMargin(6,4) )
+			[
+				SNew(STextBlock)
+				.TextStyle( FEditorStyle::Get(), TEXT("KismetExpression.ReadAutogeneratedVariable") )
+				.Text( EffectiveVariableName )
+			];
+#endif
+	}
+	else if (auto VarGetNode = Cast<const UK2Node_VariableGet>(Node))
+	{
+		// Variable get node
+		return SNew(SOverlay)
+			+ SOverlay::Slot()
+			[
+				SNew(SImage)
+				.Image( FEditorStyle::GetBrush("KismetExpression.ReadVariable.Body") )
+				.ColorAndOpacity( Schema->GetPinTypeColor(OutputPins[0]->PinType) )
+			]
+			+ SOverlay::Slot()
+			.Padding( FMargin(6,4) )
+			[
+				SNew(STextBlock)
+				.TextStyle( FEditorStyle::Get(), TEXT("KismetExpression.ReadVariable") )
+				.Text( FText::FromString( VarGetNode->GetVarNameString() ) )
+			]
+			+ SOverlay::Slot()
+			[
+				SNew(SImage)
+				.Visibility(EVisibility::HitTestInvisible)
+				.Image( FEditorStyle::GetBrush("KismetExpression.ReadVariable.Gloss") )
+			];
+		/*
+			+ SOverlay::Slot()
+			.Padding( FMargin(6,4) )
+			[
+				SNew(STextBlock)
+				.TextStyle( FEditorStyle::Get(), TEXT("KismetExpression.ReadVariable") )
+				.Text( VarGetNode->VariableName.ToString() )
+			];
+			*/
+	}
+	else if (UK2Node* AnyNode = Cast<UK2Node>(Node))
+	{
+		const bool bIsCompact = AnyNode->ShouldDrawCompact() && (InputPinCount <= 2);
+
+		TSharedRef<SWidget> OperationWidget = 
+			SNew(STextBlock)
+			.TextStyle( FEditorStyle::Get(), bIsCompact ? TEXT("KismetExpression.OperatorNode") : TEXT("KismetExpression.FunctionNode") )
+			.Text(AnyNode->GetCompactNodeTitle());
+
+		if ((InputPinCount == 1) && bIsCompact)
+		{
+			// One-pin compact nodes are assumed to be unary operators
+			return SNew(SHorizontalBox)
+				+SHorizontalBox::Slot()
+				.AutoWidth()
+				.VAlign(VAlign_Center)
+				.Padding(4, 0)
+				[
+					OperationWidget
+				]
+				+SHorizontalBox::Slot()
+				.AutoWidth()
+				.VAlign(VAlign_Center)
+				[
+					MakePinWidget(InputPins[0])
+				];
+		}
+		else if ((InputPinCount == 2) && bIsCompact)
+		{
+			// Two-pin compact nodes are assumed to be binary operators
+			return SNew(SHorizontalBox)
+				+SHorizontalBox::Slot()
+				.AutoWidth()
+				.VAlign(VAlign_Center)
+				[
+					MakePinWidget(InputPins[0])
+				]
+				+SHorizontalBox::Slot()
+				.AutoWidth()
+				.VAlign(VAlign_Center)
+				.Padding(4, 0)
+				[
+					OperationWidget
+				]
+				+SHorizontalBox::Slot()
+				.AutoWidth()
+				.VAlign(VAlign_Center)
+				[
+					MakePinWidget(InputPins[1])
+				];
+		}
+		else
+		{
+			// All other operations are treated as traditional function calls
+			return MakeFunctionCallWidget(AnyNode);
+		}
+	}
+	else
+	{
+		return SNew(STextBlock).Text(LOCTEXT("UnknownNodeMessage", "UNKNOWN_NODE"));
+	}
+}
+
+TSharedRef<SWidget> FGraphSnapContainerBuilder::MakeFunctionCallWidget(UK2Node* AnyNode)
+{
+	const UEdGraphSchema_K2* Schema = GetDefault<UEdGraphSchema_K2>();
+
+	TSharedRef<SHorizontalBox> WidgetContainer = SNew(SHorizontalBox);
+
+	// Add the return value (if any)
+	UEdGraphPin* ReturnValuePin = AnyNode->FindPin(UEdGraphSchema_K2::PN_ReturnValue);
+	UEdGraphPin* SelfPin = Schema->FindSelfPin(*AnyNode, EGPD_Input);
+
+	if (ReturnValuePin != nullptr)
+	{
+		WidgetContainer->AddSlot()
+			.AutoWidth()
+			.VAlign(VAlign_Center)
+			.Padding(4)
+			[
+				MakePinWidget(ReturnValuePin)
+			];
+
+		WidgetContainer->AddSlot()
+			.AutoWidth()
+			.VAlign(VAlign_Center)
+			.Padding(4)
+			[
+				SNew(STextBlock)
+				.Text(LOCTEXT("AssignmentExpression", "="))
+			];
+	}
+
+	// Add the scope if needed
+	if ((SelfPin != nullptr) && !SelfPin->bHidden)
+	{
+		WidgetContainer->AddSlot()
+			.AutoWidth()
+			.VAlign(VAlign_Center)
+			[
+				MakePinWidget(SelfPin)
+			];
+
+		WidgetContainer->AddSlot()
+			.AutoWidth()
+			.VAlign(VAlign_Center)
+			.Padding(FMargin(2.0f, 0.0f))
+			[
+				SNew(STextBlock).Text(LOCTEXT("ScopingOperator", "."))
+			];
+	}
+
+	// Add the function name
+	TSharedRef<SWidget> OperationWidget =
+		SNew(STextBlock)
+		.TextStyle(FEditorStyle::Get(), TEXT("KismetExpression.FunctionNode"))
+		.Text(AnyNode->GetCompactNodeTitle());
+	WidgetContainer->AddSlot()
+		.AutoWidth()
+		.VAlign(VAlign_Center)
+		[
+			OperationWidget
+		];
+
+	// Create the argument list
+	TSharedRef<SHorizontalBox> InnerBox =
+		SNew(SHorizontalBox)
+		+SHorizontalBox::Slot()
+		.AutoWidth()
+		.VAlign(VAlign_Center)
+		.Padding(4)
+		[
+			SNew(STextBlock)
+			.Text(LOCTEXT("BeginExpression", "("))
+		];
+
+	bool bFirstPin = true;
+	for (UEdGraphPin* Pin : AnyNode->Pins)
+	{
+		if (Schema->IsExecPin(*Pin) || (Pin == ReturnValuePin) || (Pin->bHidden) || (Pin == SelfPin))
+		{
+			continue;
+		}
+
+		//@TODO: This is not general enough, some pins are 'stupid'
+		// Hide the output variable pin if it's a variable set node and it has no connections, ugh!
+		if (Cast<UK2Node_VariableSet>(AnyNode) != nullptr)
+		{
+			if ((Pin->Direction == EGPD_Output) && (Pin->LinkedTo.Num() == 0))
+			{
+				continue;
+			}
+		}
+
+		// Add the comma if needed
+		if (!bFirstPin)
+		{
+			InnerBox->AddSlot()
+			.AutoWidth()
+			.VAlign(VAlign_Center)
+			.Padding(FMargin(4.0f,0.0f))
+			[
+				SNew(STextBlock).Text(LOCTEXT("NextExpression", ","))
+			];
+		}
+		bFirstPin = false;
+
+		// Add an indicator for output pins
+		if (Pin->Direction == EGPD_Output)
+		{
+			InnerBox->AddSlot()
+				.AutoWidth()
+				.VAlign(VAlign_Center)
+				.Padding(FMargin(0.0f, 0.0f, 4.0f, 0.0f))
+				[
+					SNew(STextBlock)
+					.Text(LOCTEXT("OutPin", "out"))
+				];
+		}
+		else if (Pin->PinType.bIsReference)
+		{
+			InnerBox->AddSlot()
+				.AutoWidth()
+				.VAlign(VAlign_Center)
+				.Padding(FMargin(0.0f, 0.0f, 4.0f, 0.0f))
+				[
+					SNew(STextBlock)
+					.Text(LOCTEXT("RefPin", "ref "))
+				];
+		}
+
+		// Add the pin contents
+		InnerBox->AddSlot()
+		.AutoWidth()
+		.VAlign(VAlign_Center)
+		[
+			MakePinWidget(Pin)
+		];
+	}
+
+	InnerBox->AddSlot()
+	.AutoWidth()
+	.VAlign(VAlign_Center)
+	.Padding(4)
+	[
+		SNew(STextBlock).Text(LOCTEXT("EndExpression", ")"))
+	];
+	WidgetContainer->AddSlot()
+		.AutoWidth()
+		.VAlign(VAlign_Center)
+		[
+			InnerBox
+		];
+
+	return WidgetContainer;
+}
+
+TSharedRef<SWidget> FGraphSnapContainerBuilder::MakePinWidget(UEdGraphPin* Pin)
+{
+	check(Pin);
+	
+	TSharedPtr<SWidget> ChildWidget = nullptr;
+	if (Pin->LinkedTo.Num() > 0)
+	{
+		UEdGraphPin* ChildPin = Pin->LinkedTo[0];
+		ChildWidget = MakeNodeWidget(ChildPin->GetOwningNode(), ChildPin->LinkedTo[0]);
+	}
+
+	return SNew(SGraphSnapContainerEntry, Pin->GetOwningNode(), Pin, ChildWidget);
+}
+
+/////////////////////////////////////////////////////
+// SGraphSnapContainerEntry
+
+void SGraphSnapContainerEntry::Construct(const FArguments& InArgs, UEdGraphNode* InNode, UEdGraphPin* InPin, TSharedPtr<SWidget> InChildWidget)
+{
+	TargetNode = InNode;
+	check(TargetNode);
+	TargetPin = InPin;
+	check(TargetPin);
+
+	FLinearColor BackgroundColor = FLinearColor::Red;
+	TSharedRef<SWidget> ContentToShow = InChildWidget.IsValid() ? InChildWidget.ToSharedRef() : SNullWidget::NullWidget;
+
+	BackgroundColor = TargetPin->GetSchema()->GetPinTypeColor(TargetPin->PinType);
+
+	if (!InChildWidget.IsValid())
+	{
+		// This should only be called on pins that have nothing hooked up to them (without passing in an inner widget!)
+		check(TargetPin->LinkedTo.Num() == 0);
+
+		const UEdGraphSchema_K2* Schema = GetDefault<UEdGraphSchema_K2>();
+		if (Schema->IsExecPin(*TargetPin))
+		{
+			ContentToShow = SNew(STextBlock)
+				.TextStyle(FEditorStyle::Get(), TEXT("KismetExpression.LiteralValue"))
+				.Text(LOCTEXT("SnapContainerExecutionPinDesc", "Drag nodes here"));
+				//.ColorAndOpacity(FLinearColor::Black);
+		}
+		else
+		{
+			// Try showing the literal field for the pin
+			if (TargetPin->Direction == EGPD_Input)
+			{
+				DefaultValuePinWidget = FNodeFactory::CreatePinWidget(TargetPin);
+				DefaultValuePinWidget->SetOnlyShowDefaultValue(true);
+
+				ContentToShow = DefaultValuePinWidget->GetDefaultValueWidget();
+				if (TargetPin->PinType.bIsReference)
+				{
+					DefaultValuePinWidget.Reset();
+					ContentToShow = SNullWidget::NullWidget;
+				}
+			}
+
+			// Failing that, show the type description
+			if (ContentToShow == SNullWidget::NullWidget)
+			{
+				ContentToShow = SNew(STextBlock)
+					.TextStyle(FEditorStyle::Get(), TEXT("KismetExpression.LiteralValue"))
+					.Text(GetTypeDescription(TargetPin->PinType));
+			}
+		}
+	}
+
+	SetVisibility(EVisibility::Visible);
+
+	ChildSlot
+	[
+		SNew(SBorder)
+		.BorderImage(FEditorStyle::GetBrush("PlainBorder"))
+		.BorderBackgroundColor(FSlateColor(BackgroundColor))
+		[
+			SNew(SBox)
+			.MinDesiredWidth(40.0f)
+			.MinDesiredHeight(20.0f)
+			.Padding(FMargin(0.0f, 0.0f))
+			.HAlign(HAlign_Fill)
+			.VAlign(VAlign_Fill)
+			[
+				ContentToShow
+			]
+		]
+	];
+}
+
+
+FReply SGraphSnapContainerEntry::OnDrop(const FGeometry& MyGeometry, const FDragDropEvent& DragDropEvent)
+{
+	//@TODO:
+//	bool bReadOnly = OwnerGraphPanelPtr.IsValid() ? !OwnerGraphPanelPtr.Pin()->IsGraphEditable() : false;
+	const bool bReadOnly = false;
+
+	TSharedPtr<FDragDropOperation> Operation = DragDropEvent.GetOperation();
+	if (!Operation.IsValid() || bReadOnly)
+	{
+		return FReply::Unhandled();
+	}
+
+	// Is someone dropping a connection onto this node?
+	if (Operation->IsOfType<FGraphEditorDragDropAction>())
+	{
+		TSharedPtr<FGraphEditorDragDropAction> DragConnectionOp = StaticCastSharedPtr<FGraphEditorDragDropAction>(Operation);
+
+		check(TargetNode);
+		check(TargetPin);
+		DragConnectionOp->SetHoveredGraph(nullptr);
+		DragConnectionOp->SetHoveredNode(TargetNode);
+		DragConnectionOp->SetHoveredPin(TargetPin);
+
+		//const FVector2D NodeAddPosition = NodeCoordToGraphCoord(MyGeometry.AbsoluteToLocal(DragDropEvent.GetScreenSpacePosition()));
+		const FVector2D NodeAddPosition = FVector2D(TargetNode->NodePosX, TargetNode->NodePosY) +
+			FVector2D(0.0f, 64.0f) * TargetNode->Pins.Find(TargetPin) +
+			FVector2D((TargetPin->Direction == EGPD_Output) ? 200.0f : -180.0f, 0.0f);
+
+		FReply Result = DragConnectionOp->DroppedOnPin(DragDropEvent.GetScreenSpacePosition(), NodeAddPosition);
+
+		if (Result.IsEventHandled())
+		{
+			TargetNode->GetGraph()->NotifyGraphChanged();
+
+			//@TODO: Need to tell ourselves we've changed too!
+			// This is doing so, but it's pretty hacky
+			if (UK2Node* CompositeNode = Cast<UK2Node>(TargetNode->GetGraph()->GetOuter()))
+			{
+				CompositeNode->ReconstructNode();
+				CompositeNode->GetGraph()->NotifyGraphChanged();
+			}
+		}
+		return Result;
+	}
+
+	return FReply::Unhandled();
+}
+
+TSharedPtr<IToolTip> SGraphSnapContainerEntry::GetToolTip()
+{
+	// We just masquerade as a pin, so grab the underlying pin tooltip here
+	if (DefaultValuePinWidget.IsValid())
+	{
+		return DefaultValuePinWidget->GetToolTip();
+	}
+
+	return SWidget::GetToolTip();
+}
+
+FReply SGraphSnapContainerEntry::OnDragDetected(const FGeometry& MyGeometry, const FPointerEvent& MouseEvent)
+{
+	if (MouseEvent.IsMouseButtonDown(EKeys::LeftMouseButton))
+	{
+// 		FDragConnection::FDraggedPinTable PinHandles;
+// 		PinHandles.Reserve(InStartingPins.Num());
+// 		// since the graph can be refreshed and pins can be reconstructed/replaced 
+// 		// behind the scenes, the DragDropOperation holds onto FGraphPinHandles 
+// 		// instead of direct widgets/graph-pins
+// 		for (const TSharedRef<SGraphPin>& PinWidget : InStartingPins)
+// 		{
+// 			PinHandles.Add(PinWidget->GetPinObj());
+// 		}
+// 
+// 		TSharedRef<FDragConnection> DragDropOp = FDragConnection::New(InGraphPanel, PinHandles);
+// 
+// 		return FReply::Handled().BeginDragDrop(DragDropOp);
+	}
+
+	return FReply::Unhandled();
+}
+
+/////////////////////////////////////////////////////
+
+#undef LOCTEXT_NAMESPACE
diff --git a/Engine/Plugins/Experimental/BlueprintSnapNodes/Source/BlueprintSnapNodes/Private/SGraphSnapContainerRow.h b/Engine/Plugins/Experimental/BlueprintSnapNodes/Source/BlueprintSnapNodes/Private/SGraphSnapContainerRow.h
new file mode 100644
--- /dev/null
+++ b/Engine/Plugins/Experimental/BlueprintSnapNodes/Source/BlueprintSnapNodes/Private/SGraphSnapContainerRow.h
@@ -0,0 +1,54 @@
+// Copyright Epic Games, Inc. All Rights Reserved.
+
+#pragma once
+
+#include "CoreMinimal.h"
+#include "Widgets/DeclarativeSyntaxSupport.h"
+#include "Widgets/SCompoundWidget.h"
+#include "KismetNodes/SGraphNodeK2Base.h"
+
+class SGraphPin;
+class UK2Node;
+
+/////////////////////////////////////////////////////
+// FGraphSnapContainerBuilder
+
+class FGraphSnapContainerBuilder
+{
+public:
+	static TSharedRef<SWidget> CreateSnapContainerWidgets(UEdGraph* ModelGraph, UEdGraphNode* RootNode);
+
+private:
+	FGraphSnapContainerBuilder(UEdGraph* InGraph);
+
+	TSharedRef<SWidget> MakeNodeWidget(UEdGraphNode* Node, UEdGraphPin* FromPin);
+	TSharedRef<SWidget> MakePinWidget(UEdGraphPin* Pin);
+	TSharedRef<SWidget> MakeFunctionCallWidget(UK2Node* AnyNode);
+private:
+	UEdGraph* Graph;
+	TSet<const UEdGraphNode*> VisitedNodes;
+};
+
+/////////////////////////////////////////////////////
+// SGraphSnapContainerEntry
+
+class SGraphSnapContainerEntry : public SCompoundWidget
+{
+public:
+	SLATE_BEGIN_ARGS(SGraphSnapContainerEntry){}
+	SLATE_END_ARGS()
+
+	void Construct(const FArguments& InArgs, UEdGraphNode* InNode, UEdGraphPin* InPin, TSharedPtr<SWidget> InChildWidget);
+
+	// SWidget interface
+	virtual FReply OnDragDetected(const FGeometry& MyGeometry, const FPointerEvent& MouseEvent) override;
+	virtual FReply OnDrop(const FGeometry& MyGeometry, const FDragDropEvent& DragDropEvent) override;
+	virtual TSharedPtr<IToolTip> GetToolTip() override;
+	// End of SWidget interface
+
+private:
+	UEdGraphNode* TargetNode;
+	UEdGraphPin* TargetPin;
+
+	TSharedPtr<SGraphPin> DefaultValuePinWidget;
+};
diff --git a/Engine/Plugins/Experimental/BlueprintSnapNodes/Source/BlueprintSnapNodes/Public/BlueprintSnapNodes.h b/Engine/Plugins/Experimental/BlueprintSnapNodes/Source/BlueprintSnapNodes/Public/BlueprintSnapNodes.h
new file mode 100644
--- /dev/null
+++ b/Engine/Plugins/Experimental/BlueprintSnapNodes/Source/BlueprintSnapNodes/Public/BlueprintSnapNodes.h
@@ -0,0 +1,21 @@
+// Copyright Epic Games, Inc. All Rights Reserved.
+
+#pragma once
+
+#include "CoreMinimal.h"
+#include "Modules/ModuleManager.h"
+
+class FBlueprintSnapNodeWidgetFactory;
+
+class FBlueprintSnapNodesModule : public IModuleInterface
+{
+public:
+
+	/** IModuleInterface implementation */
+	virtual void StartupModule() override;
+	virtual void ShutdownModule() override;
+
+private:
+	TSharedPtr<FBlueprintSnapNodeWidgetFactory> NodeWidgetFactory;
+
+};
