[Task] Demote BlueprintSnapNodes plugin to Experimental

diff --git a/Engine/Plugins/Experimental/BlueprintSnapNodes/BlueprintSnapNodes.uplugin b/Engine/Plugins/Experimental/BlueprintSnapNodes/BlueprintSnapNodes.uplugin
--- a/Engine/Plugins/Experimental/BlueprintSnapNodes/BlueprintSnapNodes.uplugin
+++ b/Engine/Plugins/Experimental/BlueprintSnapNodes/BlueprintSnapNodes.uplugin
@@ -11,7 +11,7 @@
 	"MarketplaceURL": "",
 	"SupportURL": "",
 	"CanContainContent": false,
-	"IsBetaVersion": true,
+	"IsBetaVersion": false,
 	"Installed": false,
 	"EnabledByDefault": false,
 	"Modules": [
@@ -21,5 +21,5 @@
 			"LoadingPhase": "Default"
 		}
 	],
-	"IsExperimentalVersion": false
+	"IsExperimentalVersion": true
 }
\ No newline at end of file
